WebAPI missing files

diff --git a/Engine/Plugins/Web/WebAPI/Source/WebAPI/Private/WebAPILog.h b/Engine/Plugins/Web/WebAPI/Source/WebAPI/Private/WebAPILog.h
new file mode 100644
--- /dev/null
+++ b/Engine/Plugins/Web/WebAPI/Source/WebAPI/Private/WebAPILog.h
@@ -0,0 +1,7 @@
+// Copyright Epic Games, Inc. All Rights Reserved.
+
+#pragma once
+
+#include "CoreMinimal.h"
+
+DECLARE_LOG_CATEGORY_EXTERN(LogWebAPI, Log, All);
diff --git a/Engine/Plugins/Web/WebAPI/Source/WebAPI/Private/WebAPIModule.cpp b/Engine/Plugins/Web/WebAPI/Source/WebAPI/Private/WebAPIModule.cpp
new file mode 100644
--- /dev/null
+++ b/Engine/Plugins/Web/WebAPI/Source/WebAPI/Private/WebAPIModule.cpp
@@ -0,0 +1,19 @@
+// Copyright Epic Games, Inc. All Rights Reserved.
+
+#include "WebAPIModule.h"
+
+#define LOCTEXT_NAMESPACE "WebAPI"
+
+void FWebAPIModule::StartupModule()
+{
+
+}
+
+void FWebAPIModule::ShutdownModule()
+{
+
+}
+
+#undef LOCTEXT_NAMESPACE
+
+IMPLEMENT_MODULE(FWebAPIModule, WebAPI)
diff --git a/Engine/Plugins/Web/WebAPI/Source/WebAPI/Private/WebAPIModule.h b/Engine/Plugins/Web/WebAPI/Source/WebAPI/Private/WebAPIModule.h
new file mode 100644
--- /dev/null
+++ b/Engine/Plugins/Web/WebAPI/Source/WebAPI/Private/WebAPIModule.h
@@ -0,0 +1,15 @@
+// Copyright Epic Games, Inc. All Rights Reserved.
+
+#pragma once
+
+#include "IWebAPIModule.h"
+
+#include "CoreMinimal.h"
+
+class FWebAPIModule final
+    : public IWebApiModuleInterface
+{
+public:
+    virtual void StartupModule() override;
+    virtual void ShutdownModule() override;
+};
diff --git a/Engine/Plugins/Web/WebAPI/Source/WebAPI/Private/WebAPIStats.cpp b/Engine/Plugins/Web/WebAPI/Source/WebAPI/Private/WebAPIStats.cpp
new file mode 100644
--- /dev/null
+++ b/Engine/Plugins/Web/WebAPI/Source/WebAPI/Private/WebAPIStats.cpp
@@ -0,0 +1,3 @@
+// Copyright Epic Games, Inc. All Rights Reserved.
+
+#include "WebAPIStats.h"
diff --git a/Engine/Plugins/Web/WebAPI/Source/WebAPI/Private/WebAPIStats.h b/Engine/Plugins/Web/WebAPI/Source/WebAPI/Private/WebAPIStats.h
new file mode 100644
--- /dev/null
+++ b/Engine/Plugins/Web/WebAPI/Source/WebAPI/Private/WebAPIStats.h
@@ -0,0 +1,10 @@
+// Copyright Epic Games, Inc. All Rights Reserved.
+
+#pragma once
+
+#include "CoreMinimal.h"
+
+#include "Stats/Stats.h"
+#include "Stats/Stats2.h"
+
+DECLARE_STATS_GROUP(TEXT("WebAPI"), STATGROUP_WebAPI, STATCAT_Advanced);
diff --git a/Engine/Plugins/Web/WebAPI/Source/WebAPI/Public/IWebAPIModule.h b/Engine/Plugins/Web/WebAPI/Source/WebAPI/Public/IWebAPIModule.h
new file mode 100644
--- /dev/null
+++ b/Engine/Plugins/Web/WebAPI/Source/WebAPI/Public/IWebAPIModule.h
@@ -0,0 +1,74 @@
+// Copyright Epic Games, Inc. All Rights Reserved.
+
+#pragma once
+
+#include "CoreMinimal.h"
+#include "WebAPIHttpRequest.h"
+#include "Modules/ModuleInterface.h"
+#include "Modules/ModuleManager.h"
+
+class IWebApiModuleInterface : public IModuleInterface
+{
+public:
+	virtual void StartupModule() override {}
+	virtual void ShutdownModule() override {}
+
+	static IWebApiModuleInterface& Get()
+	{
+		static const FName ModuleName = "WebAPI";
+		return FModuleManager::LoadModuleChecked<IWebApiModuleInterface>(ModuleName);
+	}
+
+	template <typename PayloadType>
+	static TSharedRequest<PayloadType> CreateRequest(const FString& InURL, const FString& Verb = TEXT("POST"), const FString& ContentType = TEXT("application/json"));
+
+	template<typename PayloadType>
+	static TSharedRequest<PayloadType> BatchRequests(TArray<TSharedRequest<PayloadType>>&& Requests);
+};
+
+template <typename PayloadType>
+TSharedRequest<PayloadType> IWebApiModuleInterface::CreateRequest(const FString& InURL, const FString& Verb, const FString& ContentType)
+{
+	ensure(IsInGameThread());
+
+	TSharedRef<TRequest<PayloadType>> NewRequest = MakeShared<TRequest<PayloadType>>();
+	NewRequest->InitializeRequest(InURL, Verb, ContentType);
+
+	return MoveTemp(NewRequest);
+}
+
+template <typename PayloadType>
+TSharedRequest<PayloadType> IWebApiModuleInterface::BatchRequests(TArray<TSharedRef<TRequest<PayloadType>>>&& Requests)
+{
+	ensure(IsInGameThread() && Requests.Num() > 0);
+
+	if (Requests.Num() == 1)
+	{
+		ensure(Requests[0]->HasPayload());
+		return Requests[0];
+	}
+
+	using FTRequestPayloadRef = TSharedRef<TRequest<PayloadType>>;
+	FTRequestPayloadRef BatchedRequest = TRequest<PayloadType>::CopyHeadersToNewRequest(Requests[0]);
+	
+	FString BatchedPayload = TEXT("[");	
+	for (const FTRequestPayloadRef& Request : Requests)
+	{
+		ensure(Request->HasPayload());
+		ensure(Requests[0] == Request);
+		
+		BatchedPayload += Request->PayloadToString().GetValue();
+		BatchedPayload += TEXT(',');
+	}
+
+	// replaces trailing comma with closing bracket and sets the payload
+	BatchedPayload[BatchedPayload.Len() - 1] = TEXT(']');
+	const int32 Utf8Length = FTCHARToUTF8_Convert::ConvertedLength(*BatchedPayload, BatchedPayload.Len());
+	TArray<uint8> Buffer;
+	Buffer.SetNumUninitialized(Utf8Length);
+	FTCHARToUTF8_Convert::Convert(reinterpret_cast<UTF8CHAR*>(Buffer.GetData()), Buffer.Num(), *BatchedPayload, BatchedPayload.Len());
+	
+	BatchedRequest->SetPayloadData(Buffer);
+
+	return BatchedRequest;
+}
diff --git a/Engine/Plugins/Web/WebAPI/Source/WebAPI/WebAPI.Build.cs b/Engine/Plugins/Web/WebAPI/Source/WebAPI/WebAPI.Build.cs
new file mode 100644
--- /dev/null
+++ b/Engine/Plugins/Web/WebAPI/Source/WebAPI/WebAPI.Build.cs
@@ -0,0 +1,43 @@
+// Copyright Epic Games, Inc. All Rights Reserved.
+
+using UnrealBuildTool;
+
+public class WebAPI : ModuleRules
+{
+	public WebAPI(ReadOnlyTargetRules Target) : base(Target)
+	{
+		PCHUsage = ModuleRules.PCHUsageMode.UseExplicitOrSharedPCHs;
+
+		// Allows display of more verbose information including UI tooltips, etc.
+		bool bWithWebAPIDebug = true;
+		
+		// Forcibly disables the flag if not valid for the Target
+		bWithWebAPIDebug &= Target.Type == TargetType.Editor || Target.Type == TargetType.Program;
+		PublicDefinitions.Add($"WITH_WEBAPI_DEBUG={(bWithWebAPIDebug ? 1 : 0)}");
+
+		PublicDependencyModuleNames.AddRange(
+			new string[] {
+				"Core",
+                "CoreUObject",
+                "DeveloperSettings",
+            });
+
+		PrivateDependencyModuleNames.AddRange(
+			new string[] {
+				"Engine",
+				"HTTP",
+				"Slate",
+				"SlateCore",
+				"Json",
+				"JsonUtilities"
+			});
+
+        if (Target.bBuildEditor == true)
+        {
+            PrivateDependencyModuleNames.AddRange(
+                new string[] {
+                    "UnrealEd"
+                });
+        }
+    }
+}
diff --git a/Engine/Plugins/Web/WebAPI/WebAPI.uplugin b/Engine/Plugins/Web/WebAPI/WebAPI.uplugin
new file mode 100644
--- /dev/null
+++ b/Engine/Plugins/Web/WebAPI/WebAPI.uplugin
@@ -0,0 +1,54 @@
+{
+	"FileVersion": 3,
+	"Version": 1,
+	"VersionName": "1.0",
+	"FriendlyName": "WebAPI",
+	"Description": "",
+	"Category": "Web",
+	"CreatedBy": "",
+	"CreatedByURL": "",
+	"DocsURL": "",
+	"MarketplaceURL": "",
+	"SupportURL": "",
+	"CanContainContent": true,
+	"IsBetaVersion": true,
+	"IsExperimentalVersion": false,
+	"Installed": false,
+	"Modules": [
+		{
+			"Name": "WebAPI",
+			"Type": "Runtime",
+			"LoadingPhase": "Default"
+		},
+        {
+			"Name": "WebAPIEditor",
+			"Type": "Editor",
+			"LoadingPhase": "Default"
+		},
+        {
+            "Name": "WebAPIBlueprintGraph",
+            "Type": "UncookedOnly",
+            "LoadingPhase": "Default"
+        },
+        {
+			"Name": "WebAPIOpenAPI",
+			"Type": "Editor",
+			"LoadingPhase": "Default"
+		},
+		{
+			"Name": "WebAPILiquidJS",
+			"Type": "Editor",
+			"LoadingPhase": "Default"
+		}
+	],
+	"Plugins": [
+		{
+			"Name": "PluginBrowser",
+			"Enabled": true
+		},
+		{
+			"Name": "WebSocketNetworking",
+			"Enabled": true
+		}
+	]
+}
diff --git a/Engine/Plugins/Web/WebAPI/WebAPIGeneratorApp/config.json b/Engine/Plugins/Web/WebAPI/WebAPIGeneratorApp/config.json
new file mode 100644
--- /dev/null
+++ b/Engine/Plugins/Web/WebAPI/WebAPIGeneratorApp/config.json
@@ -0,0 +1,4 @@
+{
+  "PORT": 8080,
+  "NODE_ENV": "development"
+}
\ No newline at end of file
diff --git a/Engine/Plugins/Web/WebAPI/WebAPIGeneratorApp/package.json b/Engine/Plugins/Web/WebAPI/WebAPIGeneratorApp/package.json
new file mode 100644
--- /dev/null
+++ b/Engine/Plugins/Web/WebAPI/WebAPIGeneratorApp/package.json
@@ -0,0 +1,40 @@
+{
+  "name": "web-api-generator-server",
+  "version": "0.0.1",
+  "description": "WebAPIGeneratorApp initialised with create-express-typescript-application.",
+  "main": "src/index.ts",
+  "scripts": {
+    "build": "tsc",
+    "start": "node ./build/index.js",
+    "dev": "ts-node ./src/index.ts",
+    "dev:nodemon": "tsc && nodemon -w src -w templates -e ts,json,liquid -x ts-node ./src/index.ts",
+    "find-deadcode": "ts-prune"
+  },
+  "keywords": ["express", "typescript"],
+  "devDependencies": {
+    "@types/express": "^4.17.13",
+    "@types/node": "^14.17.14",
+    "body-parser": "^1.19.0",
+    "express": "^4.17.1",
+    "fs-extra": "^9.0.1",
+    "lodash": "^4.17.20",
+    "nodemon": "^2.0.12",
+    "superagent": "^6.0.0",
+    "ts-node": "^9.1.1",
+    "typescript": "^4.4.2",
+    "ws": "^7.3.1"
+  },
+  "dependencies": {
+    "@overnightjs/core": "^1.7.6",
+    "@types/express": "^4.17.7",
+    "@types/fs-extra": "^9.0.1",
+    "@types/lodash": "^4.14.160",
+    "@types/node": "^14.6.0",
+    "@types/superagent": "4.1.10",
+    "@types/ws": "^7.2.6",  
+    "express": "^4.17.1",
+    "http-status": "^1.5.0",
+    "liquidjs": "^9.25.1",
+    "typescript": "^4.0.2"    
+  }
+}
diff --git a/Engine/Plugins/Web/WebAPI/WebAPIGeneratorApp/src/api.ts b/Engine/Plugins/Web/WebAPI/WebAPIGeneratorApp/src/api.ts
new file mode 100644
--- /dev/null
+++ b/Engine/Plugins/Web/WebAPI/WebAPIGeneratorApp/src/api.ts
@@ -0,0 +1,28 @@
+// properties required by template, may be different per-template
+interface TemplateContext {
+
+}
+
+// common properties for many contexts
+interface TemplateContextShared {
+    namespacePrefix: string // ie. SomeNamespace -> FSomeNamespace_ClassName
+}
+
+interface GenerateContext {
+    templateName: string // file name without path or extension
+    outputPath: string // can be relative or absolute
+
+    templateContext: TemplateContext | null
+}
+
+interface GenerateContextCollection {
+    outputPathCommon: string // All other paths, if relative, are relative to this
+    items: GenerateContext[]
+
+    templateContextShared: TemplateContextShared | null
+}
+
+// @note: should probably be websocket server - caller will want 
+// progress and any warnings/errors as they arise, not just a final report
+
+
diff --git a/Engine/Plugins/Web/WebAPI/WebAPIGeneratorApp/src/index.ts b/Engine/Plugins/Web/WebAPI/WebAPIGeneratorApp/src/index.ts
new file mode 100644
--- /dev/null
+++ b/Engine/Plugins/Web/WebAPI/WebAPIGeneratorApp/src/index.ts
@@ -0,0 +1,7 @@
+import { App } from './server';
+
+const app = new App();
+app.start()
+    .catch(err => {
+      console.log('Error:', err.message);
+    });
diff --git a/Engine/Plugins/Web/WebAPI/WebAPIGeneratorApp/src/server.ts b/Engine/Plugins/Web/WebAPI/WebAPIGeneratorApp/src/server.ts
new file mode 100644
--- /dev/null
+++ b/Engine/Plugins/Web/WebAPI/WebAPIGeneratorApp/src/server.ts
@@ -0,0 +1,122 @@
+import { Server } from '@overnightjs/core';
+import { Request, Response, NextFunction, static as expressStatic } from 'express';
+import bodyParser from 'body-parser';
+import path from 'path';
+import request from 'superagent';
+import { ServiceSettings } from './settings';
+import { LogServer } from './logServer';
+import { GeneratorController } from './controllers/generator.controller';
+import { ApiController } from './controllers/api.controller';
+
+export class App extends Server {
+    constructor() {
+        super();
+
+        this.app.use(bodyParser.json());
+        this.app.use(bodyParser.urlencoded({ extended: false }));
+        this.app.use((req, res, next) => App.setAccessControl(res, next));
+        this.app.use(expressStatic(path.join(ServiceSettings.rootFolder, 'public'), { setHeaders: (res) => App.setAccessControl(res) }));
+    }
+
+    public async start(): Promise<void> {
+        try {
+            const apiController = new ApiController();
+            apiController.initialize();
+
+            const generatorController = new GeneratorController();
+            generatorController.initialize();
+
+            super.addControllers([apiController, generatorController]);
+
+            this.app.use('/api', (req, res) => App.notFoundHandler(req, res, true));
+
+            // Trying to kill a zombie process
+            await request.get(`http://127.0.0.1:${ServiceSettings.port}/api/shutdown`)
+                .timeout(100)
+                .then(() => new Promise(resolve => setTimeout(resolve, 100)))
+                .catch(() => { });
+
+            await this.startServer();
+            console.log('DONE: WebApp started, port:', ServiceSettings.port);
+
+        } catch (error: any) {
+            console.error('ERROR:', error.message);
+            throw error;
+        }
+    }
+
+    private startServer(): Promise<void> {
+        return new Promise((resolve, reject) => {
+            const server = this.app.listen(ServiceSettings.port, async () => {
+                await LogServer.initialize();
+                resolve();
+            });
+
+            server.once('error', reject);
+        });
+    }
+
+    private static setAccessControl(res: Response, next?: NextFunction): void {
+        res.header('Access-Control-Allow-Origin', '*');
+        res.header('Access-Control-Allow-Headers', '*');
+        res.header('Access-Control-Max-Age', '86400');
+        res.header('Access-Control-Allow-Methods', 'GET, POST, PUT, DELETE, OPTIONS');
+
+        next?.();
+    }
+
+    private static notFoundHandler(req: Request, res: Response, isApi: boolean): void {
+        res.status(404)
+            .send({ error: `Invalid ${isApi ? 'api call' : 'request'}: ${req.method} '${req.originalUrl}'` });
+    }
+}
+
+/*
+import express, { Express } from 'express';
+import morgan from 'morgan';
+import helmet from 'helmet';
+import cors from 'cors';
+import config from '../config.json';
+import { Liquid } from 'liquidjs';
+import { request } from 'http';
+import { ApiController } from './controllers/api.controller';
+const routes = require('./routes/v1');
+console.log(path.resolve(__dirname, 'templates/'))
+
+const app: Express = express();
+const engine: Liquid = new Liquid({
+  root: path.resolve(__dirname, 'templates/'),
+  extname: '.liquid',
+  strictFilters: true
+});
+
+app.engine('liquid', engine.express())
+app.set('json spaces', 4);
+app.use(express.json());
+app.use(express.urlencoded({ extended: true }));
+
+await request.get(`https://127.0.0.1::${ServiceSettings.port}/api/${ServiceSettings.version}/shutdown`)
+
+// Handle logs in console during development
+if (process.env.NODE_ENV === 'development' || config.NODE_ENV === 'development') {
+  app.use(morgan('dev'));
+  app.use(cors());
+}
+
+// Handle security and origin in production
+if (process.env.NODE_ENV === 'production' || config.NODE_ENV === 'production') {
+  app.use(helmet());
+}
+
+app.use('/v1', routes);
+
+app.use((err: Error, req: express.Request, res: express.Response, next: express.NextFunction) => {
+  return res.status(500).json({
+    errorName: err.name,
+    message: err.message,
+    stack: err.stack || 'no stack defined'
+  });
+});
+
+export default app;
+*/
\ No newline at end of file
diff --git a/Engine/Plugins/Web/WebAPI/WebAPIGeneratorApp/tsconfig.json b/Engine/Plugins/Web/WebAPI/WebAPIGeneratorApp/tsconfig.json
new file mode 100644
--- /dev/null
+++ b/Engine/Plugins/Web/WebAPI/WebAPIGeneratorApp/tsconfig.json
@@ -0,0 +1,15 @@
+{
+  "compilerOptions": {
+    "module": "commonjs",
+    "resolveJsonModule": true,
+    "esModuleInterop": true,
+    "target": "es2018",
+    "noImplicitAny": true,
+    "moduleResolution": "node",
+    "sourceMap": false,
+    "outDir": "./build",
+    "experimentalDecorators": true,
+    "emitDecoratorMetadata": true,
+    "strict": true,
+  }
+}
